Removed StaticMeshEditorExtension in TP_MFG_CollabBP too to fix build breakage

diff --git a/Templates/TP_MFG_CollabBP/TP_MFG_CollabBP.uproject b/Templates/TP_MFG_CollabBP/TP_MFG_CollabBP.uproject
--- a/Templates/TP_MFG_CollabBP/TP_MFG_CollabBP.uproject
+++ b/Templates/TP_MFG_CollabBP/TP_MFG_CollabBP.uproject
@@ -46,10 +46,6 @@
 				"Win64"
 			]
 		},
-		{
-			"Name": "StaticMeshEditorExtension",
-			"Enabled": true
-		},
 		{
 			"Name": "DatasmithGLTFImporter",
 			"Enabled": true
