Collab Templates | Removing HoloLens from OpenXR SupportedTargetPlatforms

diff --git a/Templates/TP_MFG_CollabBP/TP_MFG_CollabBP.uproject b/Templates/TP_MFG_CollabBP/TP_MFG_CollabBP.uproject
--- a/Templates/TP_MFG_CollabBP/TP_MFG_CollabBP.uproject
+++ b/Templates/TP_MFG_CollabBP/TP_MFG_CollabBP.uproject
@@ -66,8 +66,7 @@
 			"SupportedTargetPlatforms": [
 				"Win64",
 				"Linux",
-				"Android",
-				"HoloLens"
+				"Android"
 			]
 		}
 	]
